Adding Product Configurator template.

diff --git a/Templates/TP_ProductConfigBP/Config/DefaultEngine.ini b/Templates/TP_ProductConfigBP/Config/DefaultEngine.ini
new file mode 100644
--- /dev/null
+++ b/Templates/TP_ProductConfigBP/Config/DefaultEngine.ini
@@ -0,0 +1,19 @@
+
+
+[/Script/HardwareTargeting.HardwareTargetingSettings]
+TargetedHardwareClass=Desktop
+AppliedTargetedHardwareClass=Desktop
+DefaultGraphicsPerformance=Maximum
+AppliedDefaultGraphicsPerformance=Maximum
+
+[/Script/Engine.Engine]
++ActiveGameNameRedirects=(OldGameName="TP_BlankBP",NewGameName="/Script/TP_ProductConfigBP")
++ActiveGameNameRedirects=(OldGameName="/Script/TP_BlankBP",NewGameName="/Script/TP_ProductConfigBP")
+
+[/Script/Engine.RendererSettings]
+r.DefaultFeature.AutoExposure.ExtendDefaultLuminanceRange=True
+
+[/Script/EngineSettings.GameMapsSettings]
+EditorStartupMap=/Game/ProductConfigurator/Levels/Main.Main
+GameDefaultMap=/Game/ProductConfigurator/Levels/Main.Main
+
diff --git a/Templates/TP_ProductConfigBP/Config/DefaultGame.ini b/Templates/TP_ProductConfigBP/Config/DefaultGame.ini
new file mode 100644
--- /dev/null
+++ b/Templates/TP_ProductConfigBP/Config/DefaultGame.ini
@@ -0,0 +1,6 @@
+
+
+[/Script/EngineSettings.GeneralProjectSettings]
+ProjectID=A2DCCCFF45408F9144680081DBAC8130
+ProjectName=Product Configurator
+
diff --git a/Templates/TP_ProductConfigBP/Config/TemplateDefs.ini b/Templates/TP_ProductConfigBP/Config/TemplateDefs.ini
new file mode 100644
--- /dev/null
+++ b/Templates/TP_ProductConfigBP/Config/TemplateDefs.ini
@@ -0,0 +1,96 @@
+; This file specifies how a template will be treated when creating a project from it.
+; It includes rules for ignoring files, replacing files and changing the contents of files.
+; This file is annotated to help with creating your own templates.
+
+[/Script/GameProjectGeneration.TemplateProjectDefs]
+
+; The categories that this template should appear in. The values in this list refer to the Key value given in TemplateCategories.ini.
+Categories=MFG
+
+; The localized display names for this template. Each entry specifies a Language code and Unicode text to display in the New Project dialog.
+LocalizedDisplayNames=(Language="en", Text="Product Configurator")
+
+; The localized descriptions for this template, which is displayed when the template is selected.
+LocalizedDescriptions=(Language="en", Text="This project template uses the Variant Manager, UMG and Blueprint to build a generic product configurator.\n\nThe user interface is dynamically created from the 3d model and the variants present in the level. This allows for users to swap the provided model for user-made assets and have the configurator's UI adapting to the new data.")
+; The sort key to use for the template to bypass the default alphabetical sorting.
+SortKey="3"
+
+; A list of folders to avoid copying when creating a project from this template. 
+; These should contain files only of concern to the template itself and not desired in the user's project.
+FoldersToIgnore=Media
+
+; The files to ignore copying, relative to the template root folder. 
+; %TEMPLATENAME% will be replaced by the template's folder name when resolving these. 
+; Filename comparisons are case-insensitive. 
+; For example, if this template is "TP_Blank" and the user creates "MyProject", then "TP_Blank.uproject" will be copied to "MyProject.uproject".
+FilesToIgnore="%TEMPLATENAME%.uproject"
+FilesToIgnore="%TEMPLATENAME%.png"
+FilesToIgnore="Config/TemplateDefs.ini"
+FilesToIgnore="Manifest.json"
+FilesToIgnore="contents.txt"
+
+; NOTE: For the purposes of the below examples, this template is considered to be named "TP_Blank", and the user has created a project called "MyProject".
+
+; Folders that will be renamed when creating a project from this template.
+; %PROJECTNAME% will be replaced by the project name specified by the user during creation when resolving these rules.
+; For example,  then "Source/TP_Blank" will be copied to "Source/MyProject".
+FolderRenames=(From="Source/%TEMPLATENAME%", To="Source/%PROJECTNAME%")
+FolderRenames=(From="Source/%TEMPLATENAME%Editor", To="Source/%PROJECTNAME%Editor")
+
+; Rules for renaming files when creating a project from this template. 
+; %TEMPLATENAME_UPPERCASE% and %PROJECTNAME_UPPERCASE% are the template and project names in upper case, eg. for "TP_Blank", its upper case representation is "TP_BLANK".
+; %TEMPLATENAME_LOWERCASE% and %PROJECTNAME_LOWERCASE% are the template and project names in lower case, eg. for "TP_Blank", its lower case representation is "tp_blank".
+; A rule does not have to include any of these modifiers. 
+; Extensions is a list of files to consider for these replacements. Eg. ("cpp", "h") will try to apply a rule only to files like "SomeFile.cpp" and "SomeFile.h".
+; bCaseSensitive specifies whether this replacement should be case-sensitive.
+
+; Here are some default rules for renaming files:
+
+; This rule states that a file called "TP_BLANK.cpp" will be renamed to "MYPROJECT.cpp".
+FilenameReplacements=(Extensions=("cpp","h","ini","cs"), From="%TEMPLATENAME_UPPERCASE%", To="%PROJECTNAME_UPPERCASE%", bCaseSensitive=true)
+
+; This rule states that a file called "tp_blank.ini" will be renamed to "myproject.ini"
+FilenameReplacements=(Extensions=("cpp","h","ini","cs"), From="%TEMPLATENAME_LOWERCASE%", To="%PROJECTNAME_LOWERCASE%", bCaseSensitive=true)
+
+; And this rule catches any irregularly-cased files, so that "Tp_bLANK.ini" will be renamed to "MyProject.ini".
+FilenameReplacements=(Extensions=("cpp","h","ini","cs"), From="%TEMPLATENAME%",To="%PROJECTNAME%",bCaseSensitive=false)
+
+; Rules for renaming strings in the contents of files. 
+; The same rules as above apply, but any usages of strings that match these rules inside the file will be changed. This includes all substrings!
+; For example, "TP_Blank" will be replaced by "MyProject", but "TP_Blank_Name" will also be replaced by "MyProject_Name".
+ReplacementsInFiles=(Extensions=("cpp","h","ini","cs"), From="%TEMPLATENAME_UPPERCASE%", To="%PROJECTNAME_UPPERCASE%", bCaseSensitive=true)
+ReplacementsInFiles=(Extensions=("cpp","h","ini","cs"), From="%TEMPLATENAME_LOWERCASE%", To="%PROJECTNAME_LOWERCASE%", bCaseSensitive=true)
+ReplacementsInFiles=(Extensions=("cpp","h","ini","cs"), From="%TEMPLATENAME%", To="%PROJECTNAME%", bCaseSensitive=false)
+
+; The following settings are not enabled for this project, but are included for completeness:
+
+; Whether or not this template can be used to create a project. Default: true
+; If this is set to false, then the template becomes a Feature Pack that can only be referenced from other templates, or added to a project from the Content Browser.
+;bAllowProjectCreation=false
+
+; Which settings on the project settings page should be hidden for this template. Default: none
+; Allowed values are:
+; All - Hides the settings page entirely
+; Languages - Choose between C++ and Blueprint versions.
+; HardwareTarget - Choose between Mobile and Desktop hardware targets.
+; GraphicsPreset - Choose between Maximum and Scalable graphics.
+; StarterContent - Choose whether to add starter content to the project.
+; XR - Choose whether to enable Virtual Reality support for the project.
+;HiddenSettings=All
+HiddenSettings=Languages
+HiddenSettings=HardwareTarget
+HiddenSettings=GraphicsPreset 
+;HiddenSettings=StarterContent 
+HiddenSettings=XR
+
+
+; The class types that this project uses.
+; This is only for display in the template description, and does not affect project creation.
+;ClassTypes=GameMode, Character, Actor
+
+; The asset types that this project uses.
+; This is only for display in the template description, and does not affect project creation.
+;AssetTypes=Skeleton, Static Mesh, Texture
+
+; Shared feature packs. The files in these packs listed in these structures marked as 'additionalfiles' will be copied on project generation.
+;SharedContentPacks=(MountName="2DSideScroller",DetailLevels=("Standard"))
diff --git a/Templates/TP_ProductConfigBP/Config/config.ini b/Templates/TP_ProductConfigBP/Config/config.ini
new file mode 100644
--- /dev/null
+++ b/Templates/TP_ProductConfigBP/Config/config.ini
@@ -0,0 +1,5 @@
+[Redirects]
+GameName=TP_ProductConfigBP
+
+[AdditionalFilesToAdd]
++Files=Templates/TP_ProductConfigBP/Content/*.*
\ No newline at end of file
diff --git a/Templates/TP_ProductConfigBP/TP_ProductConfigBP.uproject b/Templates/TP_ProductConfigBP/TP_ProductConfigBP.uproject
new file mode 100644
--- /dev/null
+++ b/Templates/TP_ProductConfigBP/TP_ProductConfigBP.uproject
@@ -0,0 +1,12 @@
+{
+	"FileVersion": 3,
+	"EngineAssociation": "",
+	"Category": "",
+	"Description": "",
+	"Plugins": [
+		{
+			"Name": "VariantManager",
+			"Enabled": true
+		}
+	]
+}
\ No newline at end of file
diff --git a/Templates/TP_ProductConfigBP/contents.txt b/Templates/TP_ProductConfigBP/contents.txt
new file mode 100644
--- /dev/null
+++ b/Templates/TP_ProductConfigBP/contents.txt
@@ -0,0 +1,4 @@
+../../../Templates/TP_ProductConfigBP/manifest.json
+../../../Templates/TP_ProductConfigBP/Config/config.ini
+../../../Templates/TP_ProductConfigBP/Media/TP_ProductConfigBP.png
+../../../Templates/TP_ProductConfigBP/Media/TP_ProductConfigBP_Preview.png
diff --git a/Templates/TP_ProductConfigBP/manifest.json b/Templates/TP_ProductConfigBP/manifest.json
new file mode 100644
--- /dev/null
+++ b/Templates/TP_ProductConfigBP/manifest.json
@@ -0,0 +1,31 @@
+{
+	"Version": 1,
+	"Name":
+	[
+		{
+			"Language": "en",
+			"Text": "Product Configurator"
+		}
+	],
+	"Description":
+	[
+		{
+			"Language": "en",
+			"Text": "This project template uses the Variant Manager, UMG and Blueprint to build a generic product configurator. The user interface is dynamically created from the 3d model and the variants present in the level. This allows for users to swap the provided model for user-made assets and have the configurator's UI adapting to the new data."
+		}
+	],
+	"AssetTypes": [],
+	"SearchTags": [],
+	"ClassTypes": "",
+	"Category": "BlueprintFeature",
+	"FocusAsset": "",
+	"Thumbnail": "TP_ProductConfigBP.png",
+	"Screenshots":
+	[
+		"TP_ProductConfigBP_Preview.png"
+	],
+	"SortKey": "TP_ProductConfigBP",
+	"AdditionalFeaturePacks":
+	[
+	]
+}
