Change filetype to utf8 #jira ueloc-4713

diff --git a/Templates/TP_ProductConfigBP/Config/TemplateDefs.ini b/Templates/TP_ProductConfigBP/Config/TemplateDefs.ini
--- a/Templates/TP_ProductConfigBP/Config/TemplateDefs.ini
+++ b/Templates/TP_ProductConfigBP/Config/TemplateDefs.ini
@@ -1,4 +1,4 @@
-; This file specifies how a template will be treated when creating a project from it.
+﻿; This file specifies how a template will be treated when creating a project from it.
 ; It includes rules for ignoring files, replacing files and changing the contents of files.
 ; This file is annotated to help with creating your own templates.
 
