Adding Japanese and Chinese translation #jira ueloc-4705

diff --git a/Templates/TP_ProductConfigBP/Config/TemplateDefs.ini b/Templates/TP_ProductConfigBP/Config/TemplateDefs.ini
--- a/Templates/TP_ProductConfigBP/Config/TemplateDefs.ini
+++ b/Templates/TP_ProductConfigBP/Config/TemplateDefs.ini
@@ -10,10 +10,15 @@ Categories=MFG
 ; The localized display names for this template. Each entry specifies a Language code and Unicode text to display in the New Project dialog.
 LocalizedDisplayNames=(Language="en", Text="Product Configurator")
 LocalizedDisplayNames=(Language="ko", Text="제품 컨피규레이터")
+LocalizedDisplayNames=(Language="zh-Hans", Text="产品配置器")
 
 ; The localized descriptions for this template, which is displayed when the template is selected.
 LocalizedDescriptions=(Language="en", Text="This project template uses the Variant Manager, UMG and Blueprint to build a generic product configurator.\n\nThe user interface is dynamically created from the 3d model and the variants present in the level. This allows for users to swap the provided model for user-made assets and have the configurator's UI adapting to the new data.")
 LocalizedDescriptions=(Language="ko", Text="이 프로젝트 템플릿은 일반적인 제품 컨피규레이터 제작을 위해 베리언트 매니저, UMG 그리고 블루프린트를 활용합니다.\n\n유저 인터페이스는 3D 모델과 레벨 내 베리언트를 기반으로 다이내믹하게 생성됩니다. 이를 통해 사용자들은 제공된 모델을 직접 제작한 애셋과 교체하는 것만으로 새로운 데이터에 맞는 컨피규레이터 UI를 사용할 수 있습니다.")
+LocalizedDescriptions=(Language="zh-Hans", Text="该项目模板使用变体管理器，UMG和蓝图构建了一个泛型产品配置器。\n\n用户界面是用关卡中的3D模型和变体动态创建的。因此用户能够用自制资产替换现有模型，从而使配置器的UI适配新数据。")
+LocalizedDescriptions=(Language="ja", Text="このプロジェクト テンプレートでは、Variant Manager、UMG、ブループリントを利用して、汎用的な製品コンフィギュレーターを構築できます。\n\nユーザーインターフェイスは、3D モデルとレベル内に表示されているバリアントから動的に作成されるため、ユーザーは、あらかじめ備わっているモデルをユーザー作成によるアセットと交換し、コンフィギュレーターの UI をその新たなデータに適合させることができます。")
+
+
 ; The sort key to use for the template to bypass the default alphabetical sorting.
 SortKey="3"
 
